Blueprints: Create a blueprint editor library plugin that exposes an API for blueprints to modify each other. This allows blueprints to replace any K2Nodes with another, such as upgrading UK2Node_CommutativeAssociativeBinaryOperator to a UK2Node_PromotableOperator, as well as some other basic refactoring tools like reparenting, replacing variable references, and adding or removing graphs.

diff --git a/Engine/Plugins/BlueprintEditorLibrary/BlueprintEditorLibrary.uplugin b/Engine/Plugins/BlueprintEditorLibrary/BlueprintEditorLibrary.uplugin
new file mode 100644
--- /dev/null
+++ b/Engine/Plugins/BlueprintEditorLibrary/BlueprintEditorLibrary.uplugin
@@ -0,0 +1,24 @@
+{
+	"FileVersion": 3,
+	"Version": 1,
+	"VersionName": "1.0",
+	"FriendlyName": "Blueprint Editor Library",
+	"Description": "A blueprint library that enables blueprints to edit other blueprints",
+	"Category": "Blueprints",
+	"CreatedBy": "Epic Games, Inc.",
+	"CreatedByURL": "http://epicgames.com",
+	"DocsURL": "",
+	"MarketplaceURL": "",
+	"SupportURL": "",
+	"EnabledByDefault": true,
+	"CanContainContent": true,
+	"IsBetaVersion": true,
+	"Installed": false,
+	"Modules": [
+		{
+			"Name": "BlueprintEditorLibrary",
+			"Type": "Editor",
+			"LoadingPhase": "Default"
+		}
+	]
+}
\ No newline at end of file
diff --git a/Engine/Plugins/BlueprintEditorLibrary/Source/BlueprintEditorLibrary/BlueprintEditorLibrary.Build.cs b/Engine/Plugins/BlueprintEditorLibrary/Source/BlueprintEditorLibrary/BlueprintEditorLibrary.Build.cs
new file mode 100644
--- /dev/null
+++ b/Engine/Plugins/BlueprintEditorLibrary/Source/BlueprintEditorLibrary/BlueprintEditorLibrary.Build.cs
@@ -0,0 +1,39 @@
+// Copyright Epic Games, Inc. All Rights Reserved.
+
+using UnrealBuildTool;
+
+public class BlueprintEditorLibrary : ModuleRules
+{
+	public BlueprintEditorLibrary(ReadOnlyTargetRules Target) : base(Target)
+	{
+		PCHUsage = ModuleRules.PCHUsageMode.UseExplicitOrSharedPCHs;	
+		
+		PrivateIncludePaths.AddRange(
+			new string[] 
+			{
+				"BlueprintEditorLibrary/Private",
+			}
+		);
+			
+		
+		PublicDependencyModuleNames.AddRange(
+			new string[]
+			{
+				"Core",
+				"BlueprintGraph"
+			}
+		);
+			
+		PrivateDependencyModuleNames.AddRange(
+			new string[]
+			{
+				"CoreUObject",
+				"Engine",
+				"Slate",
+				"SlateCore",
+				"UnrealEd",
+				"Kismet",
+			}
+		);
+	}
+}
diff --git a/Engine/Plugins/BlueprintEditorLibrary/Source/BlueprintEditorLibrary/Private/BlueprintEditorLibrary.cpp b/Engine/Plugins/BlueprintEditorLibrary/Source/BlueprintEditorLibrary/Private/BlueprintEditorLibrary.cpp
new file mode 100644
--- /dev/null
+++ b/Engine/Plugins/BlueprintEditorLibrary/Source/BlueprintEditorLibrary/Private/BlueprintEditorLibrary.cpp
@@ -0,0 +1,358 @@
+// Copyright Epic Games, Inc. All Rights Reserved.
+
+#include "BlueprintEditorLibrary.h"
+#include "BlueprintEditorLibraryModule.h"
+#include "Kismet2/BlueprintEditorUtils.h"
+#include "EdGraph/EdGraph.h"
+#include "EdGraph/EdGraphNode.h"
+#include "EdGraphSchema_K2.h"
+#include "EdGraphNode_Comment.h"
+#include "EdGraphUtilities.h"
+#include "BlueprintTypePromotion.h"
+#include "K2Node_PromotableOperator.h"
+#include "K2Node_CommutativeAssociativeBinaryOperator.h"
+#include "Kismet2/KismetEditorUtilities.h"
+#include "BlueprintNodeSpawner.h"
+#include "EdGraphUtilities.h"
+
+DEFINE_LOG_CATEGORY(LogBlueprintEditorLib);
+
+///////////////////////////////////////////////////////////
+// InternalBlueprintEditorLibrary
+
+namespace InternalBlueprintEditorLibrary
+{
+	/**
+	* Replace the OldNode with the NewNode and reconnect it's pins. If the pins don't
+	* exist on the NewNode, then orphan the connections.
+	*
+	* @param OldNode		The old node to replace
+	* @param NewNode		The new node to put in the old node's place
+	*/
+	static void ReplaceOldNodeWithNew(UEdGraphNode* OldNode, UEdGraphNode* NewNode)
+	{
+		const UEdGraphSchema_K2* Schema = GetDefault<UEdGraphSchema_K2>();\
+
+		if (Schema && OldNode && NewNode)
+		{
+			TMap<FName, FName> OldToNewPinMap;
+			for (UEdGraphPin* Pin : OldNode->Pins)
+			{
+				if (Pin->ParentPin != nullptr)
+				{
+					// ReplaceOldNodeWithNew() will take care of mapping split pins (as long as the parents are properly mapped)
+					continue;
+				}
+				else if (Pin->PinName == UEdGraphSchema_K2::PN_Self)
+				{
+					// there's no analogous pin, signal that we're expecting this
+					OldToNewPinMap.Add(Pin->PinName, NAME_None);
+				}
+				else
+				{
+					// The input pins follow the same naming scheme
+					OldToNewPinMap.Add(Pin->PinName, Pin->PinName);
+				}
+			}
+			
+			Schema->ReplaceOldNodeWithNew(OldNode, NewNode, OldToNewPinMap);
+		}
+	}
+};
+
+///////////////////////////////////////////////////////////
+// UBlueprintEditorLibrary
+
+UBlueprintEditorLibrary::UBlueprintEditorLibrary(const FObjectInitializer& ObjectInitializer)
+: Super(ObjectInitializer)
+{
+}
+
+void UBlueprintEditorLibrary::RenameVariableReferences(UBlueprint* Blueprint, const FName OldVarName, const FName NewVarName)
+{
+	if (!Blueprint)
+	{
+		return;
+	}
+
+	FBlueprintEditorUtils::RenameVariableReferences(Blueprint, Blueprint->GeneratedClass, OldVarName, NewVarName);
+}
+
+void UBlueprintEditorLibrary::ReplaceK2Nodes(UBlueprint* Blueprint, TSubclassOf<UEdGraphNode> OldNodeType, TSubclassOf<UEdGraphNode> NewNodeType)
+{
+	if (!Blueprint)
+	{
+		return;
+	}
+
+	TArray<UEdGraph*> AllGraphs;
+	Blueprint->GetAllGraphs(AllGraphs);
+
+	for (UEdGraph* Graph : AllGraphs)
+	{
+		Graph->Modify();
+
+		// Iterate backwards through the nodes because we will be removing some
+		for (int32 i = Graph->Nodes.Num() - 1; i >= 0; i--)
+		{
+			UEdGraphNode* OriginalNode = Graph->Nodes[i];
+			if (OriginalNode && OriginalNode->IsA(OldNodeType))
+			{
+				// Spawn a new node of the new class
+				UEdGraphNode* NewNode = NewObject<UEdGraphNode>(Graph, NewNodeType);
+				NewNode->CreateNewGuid();
+				NewNode->PostPlacedNewNode();
+				NewNode->AllocateDefaultPins();
+				NewNode->SetFlags(RF_Transactional);
+
+				NewNode->NodePosX = OriginalNode->NodePosX;
+				NewNode->NodePosY = OriginalNode->NodePosY;
+
+				InternalBlueprintEditorLibrary::ReplaceOldNodeWithNew(OriginalNode, NewNode);
+			}
+		}
+	}
+}
+
+UEdGraph* UBlueprintEditorLibrary::FindEventGraph(UBlueprint* Blueprint)
+{
+	return Blueprint ? FBlueprintEditorUtils::FindEventGraph(Blueprint) : nullptr;
+}
+
+UEdGraph* UBlueprintEditorLibrary::FindGraph(UBlueprint* Blueprint, FName GraphName)
+{
+	if (Blueprint)
+	{
+		for (UEdGraph* CurrentGraph : Blueprint->UbergraphPages)
+		{
+			if (CurrentGraph->GetFName() == GraphName)
+			{
+				return CurrentGraph;
+			}
+		}
+	}
+
+	return nullptr;
+}
+
+UK2Node_PromotableOperator* CreateOpNode(const FName OpName, UEdGraph* Graph, const int32 AdditionalPins)
+{
+	// The spawner will be null if type promo isn't enabled
+	if (UBlueprintFunctionNodeSpawner* Spawner = FTypePromotion::GetOperatorSpawner(OpName))
+	{
+		// Spawn a new node!
+		IBlueprintNodeBinder::FBindingSet Bindings;
+		FVector2D SpawnLoc{};
+		UK2Node_PromotableOperator* NewOpNode = Cast<UK2Node_PromotableOperator>(Spawner->Invoke(Graph, Bindings, SpawnLoc));
+
+		// Add the necessary number of additional pins
+		for (int32 i = 0; i < AdditionalPins; ++i)
+		{
+			NewOpNode->AddInputPin();
+		}
+
+		return NewOpNode;
+	}
+
+	return nullptr;
+}
+
+void UBlueprintEditorLibrary::UpgradeOperatorNodes(UBlueprint* Blueprint)
+{
+	if (!Blueprint)
+	{
+		return;
+	}
+
+	if (TypePromoDebug::IsTypePromoEnabled())
+	{
+		// Ensure that we have promotable operator node spanners available to us. 
+		// They will be empty if the editor hasn't been opened or 
+		if (FBlueprintActionDatabase* Actions = FBlueprintActionDatabase::TryGet())
+		{
+			Actions->RefreshAll();
+		}
+	}
+	else
+	{
+		UE_LOG(LogBlueprintEditorLib, Warning, TEXT("Type Promotion is not enabled! Cannot upgrade operator nodes. Set 'BP.TypePromo.IsEnabled' to true and try again."));
+		return;
+	}
+
+	TArray<UEdGraph*> AllGraphs;
+	Blueprint->GetAllGraphs(AllGraphs);
+	Blueprint->Modify();
+
+	for (UEdGraph* Graph : AllGraphs)
+	{
+		Graph->Modify();
+
+		for (int32 i = Graph->Nodes.Num() - 1; i >= 0; --i)
+		{
+			// Not every function that we want to upgrade is a CommunicativeBinaryOpNode
+			// Vector + Float is an example of this
+			if (UK2Node_CallFunction* OldOpNode = Cast<UK2Node_CallFunction>(Graph->Nodes[i]))
+			{
+				UFunction* Func = OldOpNode->GetTargetFunction();
+
+				// Don't bother with non-promotable functions or things that are already promotable operators
+				if (!FTypePromotion::IsPromotableFunction(Func) || OldOpNode->IsA<UK2Node_PromotableOperator>())
+				{
+					continue;
+				}
+
+				FName OpName = FTypePromotion::GetOpNameFromFunction(Func);
+
+				UK2Node_CommutativeAssociativeBinaryOperator* BinaryOpNode = Cast<UK2Node_CommutativeAssociativeBinaryOperator>(OldOpNode);
+
+				// Spawn a new node!
+				UK2Node_PromotableOperator* NewOpNode = CreateOpNode(
+					OpName,
+					OldOpNode->GetGraph(),
+					BinaryOpNode ? BinaryOpNode->GetNumberOfAdditionalInputs() : 0
+				);
+
+				// If there is a node that is a communicative op node but is not promotable
+				// then the node will be null
+				if (!NewOpNode)
+				{
+					UE_LOG(LogBlueprintEditorLib, Warning, TEXT("Failed to spawn new operator node!"));
+					continue;
+				}
+
+				NewOpNode->NodePosX = OldOpNode->NodePosX;
+				NewOpNode->NodePosY = OldOpNode->NodePosY;
+
+				InternalBlueprintEditorLibrary::ReplaceOldNodeWithNew(OldOpNode, NewOpNode);
+			}
+		}
+	}
+}
+
+void UBlueprintEditorLibrary::CompileBlueprint(UBlueprint* Blueprint)
+{
+	if (Blueprint)
+	{
+		// Skip saving this to avoid possible tautologies when saving and allow the user to manually save
+		EBlueprintCompileOptions Flags = EBlueprintCompileOptions::SkipSave;
+		FKismetEditorUtilities::CompileBlueprint(Blueprint, Flags);
+	}
+}
+
+UEdGraph* UBlueprintEditorLibrary::AddFunctionGraph(UBlueprint* Blueprint, const FString& FuncName)
+{
+	if (!Blueprint)
+	{
+		return nullptr;
+	}
+
+	Blueprint->Modify();
+	UEdGraph* NewGraph = FBlueprintEditorUtils::CreateNewGraph(
+		Blueprint, 
+		FBlueprintEditorUtils::FindUniqueKismetName(Blueprint, FuncName),
+		UEdGraph::StaticClass(), 
+		UEdGraphSchema_K2::StaticClass()
+	);
+
+	FBlueprintEditorUtils::AddFunctionGraph<UFunction>(Blueprint, NewGraph, /* bIsUserCreated = */ true, /* SignatureFromObject = */ nullptr);
+
+	return NewGraph;
+}
+
+void UBlueprintEditorLibrary::RemoveFunctionGraph(UBlueprint* Blueprint, FName FuncName)
+{
+	if (!Blueprint)
+	{
+		return;
+	}
+
+	// Find the function graph of this name
+	UEdGraph* FunctionGraph = nullptr;
+	for (UEdGraph* Graph : Blueprint->FunctionGraphs)
+	{
+		if (Graph->GetFName() == FuncName)
+		{
+			FunctionGraph = Graph;
+			break;
+		}
+	}
+
+	// Remove the function graph if we can
+	if (FunctionGraph && FunctionGraph->bAllowDeletion)
+	{
+		Blueprint->Modify();
+		FBlueprintEditorUtils::RemoveGraph(Blueprint, FunctionGraph, EGraphRemoveFlags::MarkTransient);
+	}
+	else
+	{
+		UE_LOG(LogBlueprintEditorLib, Warning, TEXT("Failed to remove function '%s' on blueprint '%s'!"), *FuncName.ToString(), *Blueprint->GetFriendlyName());
+	}
+}
+
+void UBlueprintEditorLibrary::RemoveGraph(UBlueprint* Blueprint, UEdGraph* Graph)
+{
+	if (!Blueprint || !Graph)
+	{
+		return;
+	}
+
+	FBlueprintEditorUtils::RemoveGraph(Blueprint, Graph, EGraphRemoveFlags::MarkTransient);
+}
+
+void UBlueprintEditorLibrary::RenameGraph(UEdGraph* Graph, const FString& NewNameStr)
+{
+	if (!Graph)
+	{
+		return;
+	}
+		
+	FBlueprintEditorUtils::RenameGraph(Graph, NewNameStr);
+}
+
+void UBlueprintEditorLibrary::AddComponent(UBlueprint* Blueprint, TSubclassOf<UActorComponent> ComponentClass)
+{
+	if (!Blueprint || !ComponentClass)
+	{
+		return;
+	}
+	//UClass* ActorClass = Blueprint->GeneratedClass;
+	//TSubclassOf<UActorComponent> MatchingComponentClassForAsset = FComponentAssetBrokerage::GetPrimaryComponentForAsset(AssetClass);
+
+	//AActor* CDO = Blueprint->GeneratedClass->
+	// #TODO_BH Will need to discuss how to handle this correctly, there is a lot of
+	// logic in SCS editor that I have questions about
+}
+
+UBlueprint* UBlueprintEditorLibrary::GetBlueprintAsset(UObject* Object)
+{
+	return Cast<UBlueprint>(Object);
+}
+
+void UBlueprintEditorLibrary::ReparentBlueprint(UBlueprint* Blueprint, UClass* NewParentClass)
+{
+	if (!Blueprint || !NewParentClass)
+	{
+		UE_LOG(LogBlueprintEditorLib, Warning, TEXT("Failed to reparent blueprint!"));
+		return;
+	}
+
+	if (NewParentClass == Blueprint->ParentClass)
+	{
+		UE_LOG(LogBlueprintEditorLib, Warning, TEXT("'%s' is already parented to class '%s'!"), *Blueprint->GetFriendlyName(), *NewParentClass->GetName());
+		return;
+	}
+
+	// There could be possible data loss if reparenting outside the current class hierarchy
+	if (!Blueprint->ParentClass || !NewParentClass->GetDefaultObject()->IsA(Blueprint->ParentClass))
+	{
+		UE_LOG(LogBlueprintEditorLib, Warning, TEXT("'%s' class heirarcy is changing, there could be possible data loss!"), *Blueprint->GetFriendlyName());
+	}
+
+	UClass* OriginalParentClass = Blueprint->ParentClass;
+	Blueprint->ParentClass = NewParentClass;
+
+	FBlueprintEditorUtils::RefreshAllNodes(Blueprint);
+	FBlueprintEditorUtils::MarkBlueprintAsModified(Blueprint);
+
+	CompileBlueprint(Blueprint);
+}
\ No newline at end of file
diff --git a/Engine/Plugins/BlueprintEditorLibrary/Source/BlueprintEditorLibrary/Public/BlueprintEditorLibrary.h b/Engine/Plugins/BlueprintEditorLibrary/Source/BlueprintEditorLibrary/Public/BlueprintEditorLibrary.h
new file mode 100644
--- /dev/null
+++ b/Engine/Plugins/BlueprintEditorLibrary/Source/BlueprintEditorLibrary/Public/BlueprintEditorLibrary.h
@@ -0,0 +1,143 @@
+// Copyright Epic Games, Inc. All Rights Reserved.
+
+#pragma once
+
+#include "CoreMinimal.h"
+#include "UObject/ObjectMacros.h"
+#include "Kismet/BlueprintFunctionLibrary.h"
+#include "BlueprintEditorLibrary.generated.h"
+
+DECLARE_LOG_CATEGORY_EXTERN(LogBlueprintEditorLib, Warning, All);
+
+class UActorComponent;
+
+UCLASS()
+class UBlueprintEditorLibrary : public UBlueprintFunctionLibrary
+{
+	GENERATED_UCLASS_BODY()
+public:
+
+	/**
+	* Replace any references of variables with the OldVarName to references of those with the NewVarName if possible
+	*
+	* @param Blueprint		Blueprint to replace the variable references on
+	* @param OldVarName		The variable you want replaced
+	* @param NewVarName		The new variable that will be used in the old one's place
+	*/
+	UFUNCTION(BlueprintCallable, Category = "BlueprintUpgadeTools")
+	static void RenameVariableReferences(UBlueprint* Blueprint, const FName OldVarName, const FName NewVarName);
+
+	/**
+	* Replace any existing k2node class with the new one, attempting to remap pins as we go
+	*
+	* @param Blueprint		Blueprint to replace node types on
+	* @param OldNodeType	Old node type that will be removed
+	* @param NewNodeType	New node type that will be put in the old node's place
+	*/
+	UFUNCTION(BlueprintCallable, Category = "BlueprintUpgadeTools")
+	static void ReplaceK2Nodes(UBlueprint* Blueprint, TSubclassOf<UEdGraphNode> OldNodeType, TSubclassOf<UEdGraphNode> NewNodeType);
+
+	/**
+	* Finds the event graph of the given blueprint. Null if it doesn't have one.
+	*
+	* @param Blueprint		Blueprint to search for the event graph on
+	*
+	* @return UEdGraph*		Event graph of the blueprint if it has one, null if it doesn't have one
+	*/
+	UFUNCTION(BlueprintCallable, Category = "BlueprintUpgadeTools")
+	static UEdGraph* FindEventGraph(UBlueprint* Blueprint);
+
+	/**
+	* Finds the graph with the given name on the blueprint. Null if it doesn't have one. 
+	*
+	* @param Blueprint		Blueprint to search
+	* @param GraphName		The name of the graph to search for 
+	*
+	* @return UEdGraph*		Pointer to the graph with the given name, null if not found
+	*/
+	UFUNCTION(BlueprintCallable, Category = "BlueprintUpgadeTools")
+	static UEdGraph* FindGraph(UBlueprint* Blueprint, FName GraphName);
+
+	/**
+	* Replace any UK2Node_CommutativeAssociativeBinaryOperator nodes with UK2Node_PromotableOperator nodes
+	* in the given blueprint and remap any connections they had
+	*
+	* @param Blueprint	Blueprint to upgrade
+	*/
+	UFUNCTION(BlueprintCallable, Category = "BlueprintUpgadeTools")
+	static void UpgradeOperatorNodes(UBlueprint* Blueprint);
+
+	/**
+	* Compiles the given blueprint. 
+	*
+	* @param Blueprint	Blueprint to compile
+	*/
+	UFUNCTION(BlueprintCallable, Category = "BlueprintUpgadeTools")
+	static void CompileBlueprint(UBlueprint* Blueprint);
+
+	/**
+	* Adds a function to the given blueprint
+	*
+	* @param Blueprint	The blueprint to remove the function from
+	* @param FuncName	Name of the function to add
+	*
+	* @return UEdGraph*
+	*/
+	UFUNCTION(BlueprintCallable, Category = "BlueprintUpgadeTools")
+	static UEdGraph* AddFunctionGraph(UBlueprint* Blueprint, const FString& FuncName);
+
+	/** 
+	* Deletes the function of the given name on this blueprint. Does NOT replace function call sites. 
+	*
+	* @param Blueprint		The blueprint to remove the function from
+	* @param FuncName		The name of the function to remove
+	*/
+	UFUNCTION(BlueprintCallable, Category = "BlueprintUpgadeTools")
+	static void RemoveFunctionGraph(UBlueprint* Blueprint, FName FuncName);
+
+	/** 
+	* Removes the given graph from the blueprint if possible 
+	* 
+	* @param Blueprint	The blueprint the graph will be removed from
+	* @param Graph		The graph to remove
+	*/
+	UFUNCTION(BlueprintCallable, Category = "BlueprintUpgadeTools")
+	static void RemoveGraph(UBlueprint* Blueprint, UEdGraph* Graph);
+
+	/**
+	* Attempts to rename the given graph with a new name
+	*
+	* @param Graph			The graph to rename
+	* @param NewNameStr		The new name of the graph
+	*/
+	UFUNCTION(BlueprintCallable, Category = "BlueprintUpgadeTools")
+	static void RenameGraph(UEdGraph* Graph, const FString& NewNameStr);
+
+	/**
+	* Add an Actor Component to the given blueprint
+	*
+	* @param Blueprint			Blueprint to add the component to
+	* @param ComponentClass		Type of component to add
+	*/
+	UFUNCTION(BlueprintCallable, Category = "BlueprintUpgradeTools")
+	static void AddComponent(UBlueprint* Blueprint, TSubclassOf<UActorComponent> ComponentClass);
+
+	/**
+	* Gets the UBlueprint version of the given object if possible.
+	*
+	* @param Object			The object we need to get the UBlueprint from
+	*
+	* @return UBlueprint*	The blueprint type of the given object, nullptr if the object is not a blueprint.
+	*/
+	UFUNCTION(BlueprintCallable, Category = "BlueprintUpgradeTools")
+	static UBlueprint* GetBlueprintAsset(UObject* Object);
+
+	/**
+	* Attempts to reparent the given blueprint to the new chosen parent class. 
+	*
+	* @param Blueprint			Blueprint that you would like to reparent
+	* @param NewParentClass		The new parent class to use
+	*/
+	UFUNCTION(BlueprintCallable, Category = "BlueprintUpgradeTools")
+	static void ReparentBlueprint(UBlueprint* Blueprint, UClass* NewParentClass);
+};
\ No newline at end of file
diff --git a/Engine/Plugins/BlueprintEditorLibrary/Source/BlueprintEditorLibrary/Public/BlueprintEditorLibraryModule.h b/Engine/Plugins/BlueprintEditorLibrary/Source/BlueprintEditorLibrary/Public/BlueprintEditorLibraryModule.h
new file mode 100644
--- /dev/null
+++ b/Engine/Plugins/BlueprintEditorLibrary/Source/BlueprintEditorLibrary/Public/BlueprintEditorLibraryModule.h
@@ -0,0 +1,14 @@
+// Copyright Epic Games, Inc. All Rights Reserved.
+
+#pragma once
+
+#include "Modules/ModuleManager.h"
+
+class FBlueprintEditorLibraryModule : public IModuleInterface
+{
+public:
+
+	/** IModuleInterface implementation */
+	virtual void StartupModule() override;
+	virtual void ShutdownModule() override;
+};
\ No newline at end of file
