Add missing public include path for BP editor library and add missing BlueprintEditorLibraryModule.cpp

#rnx

diff --git a/Engine/Plugins/BlueprintEditorLibrary/Source/BlueprintEditorLibrary/BlueprintEditorLibrary.Build.cs b/Engine/Plugins/BlueprintEditorLibrary/Source/BlueprintEditorLibrary/BlueprintEditorLibrary.Build.cs
--- a/Engine/Plugins/BlueprintEditorLibrary/Source/BlueprintEditorLibrary/BlueprintEditorLibrary.Build.cs
+++ b/Engine/Plugins/BlueprintEditorLibrary/Source/BlueprintEditorLibrary/BlueprintEditorLibrary.Build.cs
@@ -1,13 +1,15 @@
 // Copyright Epic Games, Inc. All Rights Reserved.
 
 using UnrealBuildTool;
+using System.IO;
 
 public class BlueprintEditorLibrary : ModuleRules
 {
 	public BlueprintEditorLibrary(ReadOnlyTargetRules Target) : base(Target)
 	{
-		PCHUsage = ModuleRules.PCHUsageMode.UseExplicitOrSharedPCHs;	
-		
+		PCHUsage = ModuleRules.PCHUsageMode.UseExplicitOrSharedPCHs;
+		PublicIncludePaths.Add(Path.Combine(ModuleDirectory, "Public"));
+
 		PrivateIncludePaths.AddRange(
 			new string[] 
 			{
diff --git a/Engine/Plugins/BlueprintEditorLibrary/Source/BlueprintEditorLibrary/Private/BlueprintEditorLibraryModule.cpp b/Engine/Plugins/BlueprintEditorLibrary/Source/BlueprintEditorLibrary/Private/BlueprintEditorLibraryModule.cpp
new file mode 100644
--- /dev/null
+++ b/Engine/Plugins/BlueprintEditorLibrary/Source/BlueprintEditorLibrary/Private/BlueprintEditorLibraryModule.cpp
@@ -0,0 +1,20 @@
+// Copyright Epic Games, Inc. All Rights Reserved.
+
+#include "BlueprintEditorLibraryModule.h"
+
+#define LOCTEXT_NAMESPACE "FBlueprintEditorLibraryModule"
+
+void FBlueprintEditorLibraryModule::StartupModule()
+{
+	// This code will execute after your module is loaded into memory; the exact timing is specified in the .uplugin file per-module
+}
+
+void FBlueprintEditorLibraryModule::ShutdownModule()
+{
+	// This function may be called during shutdown to clean up your module.  For modules that support dynamic reloading,
+	// we call this function before unloading the module.	
+}
+
+#undef LOCTEXT_NAMESPACE
+
+IMPLEMENT_MODULE(FBlueprintEditorLibraryModule, BlueprintEditorLibrary)
\ No newline at end of file
