Fix typo in Upgrade tool category name

#rnx

diff --git a/Engine/Plugins/BlueprintEditorLibrary/Source/BlueprintEditorLibrary/Public/BlueprintEditorLibrary.h b/Engine/Plugins/BlueprintEditorLibrary/Source/BlueprintEditorLibrary/Public/BlueprintEditorLibrary.h
--- a/Engine/Plugins/BlueprintEditorLibrary/Source/BlueprintEditorLibrary/Public/BlueprintEditorLibrary.h
+++ b/Engine/Plugins/BlueprintEditorLibrary/Source/BlueprintEditorLibrary/Public/BlueprintEditorLibrary.h
@@ -24,7 +24,7 @@ public:
 	* @param OldVarName		The variable you want replaced
 	* @param NewVarName		The new variable that will be used in the old one's place
 	*/
-	UFUNCTION(BlueprintCallable, Category = "BlueprintUpgadeTools")
+	UFUNCTION(BlueprintCallable, Category = "Blueprint Upgrade Tools")
 	static void RenameVariableReferences(UBlueprint* Blueprint, const FName OldVarName, const FName NewVarName);
 
 	/**
@@ -34,7 +34,7 @@ public:
 	* @param OldNodeType	Old node type that will be removed
 	* @param NewNodeType	New node type that will be put in the old node's place
 	*/
-	UFUNCTION(BlueprintCallable, Category = "BlueprintUpgadeTools")
+	UFUNCTION(BlueprintCallable, Category = "Blueprint Upgrade Tools")
 	static void ReplaceK2Nodes(UBlueprint* Blueprint, TSubclassOf<UEdGraphNode> OldNodeType, TSubclassOf<UEdGraphNode> NewNodeType);
 
 	/**
@@ -44,7 +44,7 @@ public:
 	*
 	* @return UEdGraph*		Event graph of the blueprint if it has one, null if it doesn't have one
 	*/
-	UFUNCTION(BlueprintCallable, Category = "BlueprintUpgadeTools")
+	UFUNCTION(BlueprintCallable, Category = "Blueprint Upgrade Tools")
 	static UEdGraph* FindEventGraph(UBlueprint* Blueprint);
 
 	/**
@@ -55,7 +55,7 @@ public:
 	*
 	* @return UEdGraph*		Pointer to the graph with the given name, null if not found
 	*/
-	UFUNCTION(BlueprintCallable, Category = "BlueprintUpgadeTools")
+	UFUNCTION(BlueprintCallable, Category = "Blueprint Upgrade Tools")
 	static UEdGraph* FindGraph(UBlueprint* Blueprint, FName GraphName);
 
 	/**
@@ -64,7 +64,7 @@ public:
 	*
 	* @param Blueprint	Blueprint to upgrade
 	*/
-	UFUNCTION(BlueprintCallable, Category = "BlueprintUpgadeTools")
+	UFUNCTION(BlueprintCallable, Category = "Blueprint Upgrade Tools")
 	static void UpgradeOperatorNodes(UBlueprint* Blueprint);
 
 	/**
@@ -72,7 +72,7 @@ public:
 	*
 	* @param Blueprint	Blueprint to compile
 	*/
-	UFUNCTION(BlueprintCallable, Category = "BlueprintUpgadeTools")
+	UFUNCTION(BlueprintCallable, Category = "Blueprint Upgrade Tools")
 	static void CompileBlueprint(UBlueprint* Blueprint);
 
 	/**
@@ -83,7 +83,7 @@ public:
 	*
 	* @return UEdGraph*
 	*/
-	UFUNCTION(BlueprintCallable, Category = "BlueprintUpgadeTools")
+	UFUNCTION(BlueprintCallable, Category = "Blueprint Upgrade Tools")
 	static UEdGraph* AddFunctionGraph(UBlueprint* Blueprint, const FString& FuncName);
 
 	/** 
@@ -92,7 +92,7 @@ public:
 	* @param Blueprint		The blueprint to remove the function from
 	* @param FuncName		The name of the function to remove
 	*/
-	UFUNCTION(BlueprintCallable, Category = "BlueprintUpgadeTools")
+	UFUNCTION(BlueprintCallable, Category = "Blueprint Upgrade Tools")
 	static void RemoveFunctionGraph(UBlueprint* Blueprint, FName FuncName);
 
 	/** 
@@ -101,7 +101,7 @@ public:
 	* @param Blueprint	The blueprint the graph will be removed from
 	* @param Graph		The graph to remove
 	*/
-	UFUNCTION(BlueprintCallable, Category = "BlueprintUpgadeTools")
+	UFUNCTION(BlueprintCallable, Category = "Blueprint Upgrade Tools")
 	static void RemoveGraph(UBlueprint* Blueprint, UEdGraph* Graph);
 
 	/**
@@ -110,7 +110,7 @@ public:
 	* @param Graph			The graph to rename
 	* @param NewNameStr		The new name of the graph
 	*/
-	UFUNCTION(BlueprintCallable, Category = "BlueprintUpgadeTools")
+	UFUNCTION(BlueprintCallable, Category = "Blueprint Upgrade Tools")
 	static void RenameGraph(UEdGraph* Graph, const FString& NewNameStr);
 
 	/**
@@ -119,7 +119,7 @@ public:
 	* @param Blueprint			Blueprint to add the component to
 	* @param ComponentClass		Type of component to add
 	*/
-	UFUNCTION(BlueprintCallable, Category = "BlueprintUpgradeTools")
+	UFUNCTION(BlueprintCallable, Category = "Blueprint Upgrade Tools")
 	static void AddComponent(UBlueprint* Blueprint, TSubclassOf<UActorComponent> ComponentClass);
 
 	/**
@@ -129,7 +129,7 @@ public:
 	*
 	* @return UBlueprint*	The blueprint type of the given object, nullptr if the object is not a blueprint.
 	*/
-	UFUNCTION(BlueprintCallable, Category = "BlueprintUpgradeTools")
+	UFUNCTION(BlueprintCallable, Category = "Blueprint Upgrade Tools")
 	static UBlueprint* GetBlueprintAsset(UObject* Object);
 
 	/**
@@ -138,6 +138,6 @@ public:
 	* @param Blueprint			Blueprint that you would like to reparent
 	* @param NewParentClass		The new parent class to use
 	*/
-	UFUNCTION(BlueprintCallable, Category = "BlueprintUpgradeTools")
+	UFUNCTION(BlueprintCallable, Category = "Blueprint Upgrade Tools")
 	static void ReparentBlueprint(UBlueprint* Blueprint, UClass* NewParentClass);
 };
\ No newline at end of file
