Blueprints: Fix non-unity compile error in blueprint editor library. Add defensive null check on the graph during blueprint upgrading to avoid a possible crash.

Remove AddComponent function because it isn't useable right now until we refactor the SCS tree.

#rnx

diff --git a/Engine/Plugins/BlueprintEditorLibrary/Source/BlueprintEditorLibrary/Private/BlueprintEditorLibrary.cpp b/Engine/Plugins/BlueprintEditorLibrary/Source/BlueprintEditorLibrary/Private/BlueprintEditorLibrary.cpp
--- a/Engine/Plugins/BlueprintEditorLibrary/Source/BlueprintEditorLibrary/Private/BlueprintEditorLibrary.cpp
+++ b/Engine/Plugins/BlueprintEditorLibrary/Source/BlueprintEditorLibrary/Private/BlueprintEditorLibrary.cpp
@@ -56,6 +56,9 @@ namespace InternalBlueprintEditorLibrary
 			}
 			
 			Schema->ReplaceOldNodeWithNew(OldNode, NewNode, OldToNewPinMap);
+			// reconstructing the node will clean up any
+			// incorrect default values that may have been copied over
+			NewNode->ReconstructNode();
 		}
 	}
 };
@@ -137,6 +140,11 @@ UEdGraph* UBlueprintEditorLibrary::FindGraph(UBlueprint* Blueprint, FName GraphN
 
 UK2Node_PromotableOperator* CreateOpNode(const FName OpName, UEdGraph* Graph, const int32 AdditionalPins)
 {
+	if (!Graph)
+	{
+		return nullptr;
+	}
+
 	// The spawner will be null if type promo isn't enabled
 	if (UBlueprintFunctionNodeSpawner* Spawner = FTypePromotion::GetOperatorSpawner(OpName))
 	{
@@ -144,6 +152,7 @@ UK2Node_PromotableOperator* CreateOpNode(const FName OpName, UEdGraph* Graph, co
 		IBlueprintNodeBinder::FBindingSet Bindings;
 		FVector2D SpawnLoc{};
 		UK2Node_PromotableOperator* NewOpNode = Cast<UK2Node_PromotableOperator>(Spawner->Invoke(Graph, Bindings, SpawnLoc));
+		check(NewOpNode);
 
 		// Add the necessary number of additional pins
 		for (int32 i = 0; i < AdditionalPins; ++i)
@@ -185,16 +194,22 @@ void UBlueprintEditorLibrary::UpgradeOperatorNodes(UBlueprint* Blueprint)
 
 	for (UEdGraph* Graph : AllGraphs)
 	{
+		check(Graph);
+
 		Graph->Modify();
 
 		for (int32 i = Graph->Nodes.Num() - 1; i >= 0; --i)
 		{
 			// Not every function that we want to upgrade is a CommunicativeBinaryOpNode
-			// Vector + Float is an example of this
+			// Some are just regular CallFunction nodes; Vector + Float is an example of this
 			if (UK2Node_CallFunction* OldOpNode = Cast<UK2Node_CallFunction>(Graph->Nodes[i]))
 			{
 				UFunction* Func = OldOpNode->GetTargetFunction();
+				UEdGraph* OwningGraph = OldOpNode->GetGraph();
 
+				// We should only be modifying nodes within the graph that we want
+				ensure(OwningGraph == Graph);
+				
 				// Don't bother with non-promotable functions or things that are already promotable operators
 				if (!FTypePromotion::IsPromotableFunction(Func) || OldOpNode->IsA<UK2Node_PromotableOperator>())
 				{
@@ -208,7 +223,7 @@ void UBlueprintEditorLibrary::UpgradeOperatorNodes(UBlueprint* Blueprint)
 				// Spawn a new node!
 				UK2Node_PromotableOperator* NewOpNode = CreateOpNode(
 					OpName,
-					OldOpNode->GetGraph(),
+					OwningGraph,
 					BinaryOpNode ? BinaryOpNode->GetNumberOfAdditionalInputs() : 0
 				);
 
diff --git a/Engine/Plugins/BlueprintEditorLibrary/Source/BlueprintEditorLibrary/Public/BlueprintEditorLibrary.h b/Engine/Plugins/BlueprintEditorLibrary/Source/BlueprintEditorLibrary/Public/BlueprintEditorLibrary.h
--- a/Engine/Plugins/BlueprintEditorLibrary/Source/BlueprintEditorLibrary/Public/BlueprintEditorLibrary.h
+++ b/Engine/Plugins/BlueprintEditorLibrary/Source/BlueprintEditorLibrary/Public/BlueprintEditorLibrary.h
@@ -5,12 +5,12 @@
 #include "CoreMinimal.h"
 #include "UObject/ObjectMacros.h"
 #include "Kismet/BlueprintFunctionLibrary.h"
+#include "Components/ActorComponent.h"
+#include "EdGraph/EdGraphNode.h"
 #include "BlueprintEditorLibrary.generated.h"
 
 DECLARE_LOG_CATEGORY_EXTERN(LogBlueprintEditorLib, Warning, All);
 
-class UActorComponent;
-
 UCLASS()
 class UBlueprintEditorLibrary : public UBlueprintFunctionLibrary
 {
