Blueprints: Update tooltip on UpgradeOperatorNodes to be more readable

#rnx

diff --git a/Engine/Plugins/BlueprintEditorLibrary/Source/BlueprintEditorLibrary/Public/BlueprintEditorLibrary.h b/Engine/Plugins/BlueprintEditorLibrary/Source/BlueprintEditorLibrary/Public/BlueprintEditorLibrary.h
--- a/Engine/Plugins/BlueprintEditorLibrary/Source/BlueprintEditorLibrary/Public/BlueprintEditorLibrary.h
+++ b/Engine/Plugins/BlueprintEditorLibrary/Source/BlueprintEditorLibrary/Public/BlueprintEditorLibrary.h
@@ -59,8 +59,9 @@ public:
 	static UEdGraph* FindGraph(UBlueprint* Blueprint, FName GraphName);
 
 	/**
-	* Replace any UK2Node_CommutativeAssociativeBinaryOperator nodes with UK2Node_PromotableOperator nodes
-	* in the given blueprint and remap any connections they had
+	* Replace any old operator nodes (float + float, vector + float, int + vector, etc)
+	* with the newer Promotable Operator version of the node. Preserve any connections the
+	* original node had to the newer version of the node. 
 	*
 	* @param Blueprint	Blueprint to upgrade
 	*/
