Remove incorrect '\' at the end of a line

#rnx

diff --git a/Engine/Plugins/BlueprintEditorLibrary/Source/BlueprintEditorLibrary/Private/BlueprintEditorLibrary.cpp b/Engine/Plugins/BlueprintEditorLibrary/Source/BlueprintEditorLibrary/Private/BlueprintEditorLibrary.cpp
--- a/Engine/Plugins/BlueprintEditorLibrary/Source/BlueprintEditorLibrary/Private/BlueprintEditorLibrary.cpp
+++ b/Engine/Plugins/BlueprintEditorLibrary/Source/BlueprintEditorLibrary/Private/BlueprintEditorLibrary.cpp
@@ -31,7 +31,7 @@ namespace InternalBlueprintEditorLibrary
 	*/
 	static void ReplaceOldNodeWithNew(UEdGraphNode* OldNode, UEdGraphNode* NewNode)
 	{
-		const UEdGraphSchema_K2* Schema = GetDefault<UEdGraphSchema_K2>();\
+		const UEdGraphSchema_K2* Schema = GetDefault<UEdGraphSchema_K2>();
 
 		if (Schema && OldNode && NewNode)
 		{
