Check if a function is within the operator table instead of checking all of its metadata when doing a BP upgrade

#rnx

diff --git a/Engine/Plugins/BlueprintEditorLibrary/Source/BlueprintEditorLibrary/Private/BlueprintEditorLibrary.cpp b/Engine/Plugins/BlueprintEditorLibrary/Source/BlueprintEditorLibrary/Private/BlueprintEditorLibrary.cpp
--- a/Engine/Plugins/BlueprintEditorLibrary/Source/BlueprintEditorLibrary/Private/BlueprintEditorLibrary.cpp
+++ b/Engine/Plugins/BlueprintEditorLibrary/Source/BlueprintEditorLibrary/Private/BlueprintEditorLibrary.cpp
@@ -211,7 +211,7 @@ void UBlueprintEditorLibrary::UpgradeOperatorNodes(UBlueprint* Blueprint)
 				ensure(OwningGraph == Graph);
 				
 				// Don't bother with non-promotable functions or things that are already promotable operators
-				if (!FTypePromotion::IsPromotableFunction(Func) || OldOpNode->IsA<UK2Node_PromotableOperator>())
+				if (!FTypePromotion::IsFunctionPromotionReady(Func) || OldOpNode->IsA<UK2Node_PromotableOperator>())
 				{
 					continue;
 				}
