Add a blueprint editor library function to remove unused nodes from a blueprint

#rnx

diff --git a/Engine/Plugins/BlueprintEditorLibrary/Source/BlueprintEditorLibrary/BlueprintEditorLibrary.Build.cs b/Engine/Plugins/BlueprintEditorLibrary/Source/BlueprintEditorLibrary/BlueprintEditorLibrary.Build.cs
--- a/Engine/Plugins/BlueprintEditorLibrary/Source/BlueprintEditorLibrary/BlueprintEditorLibrary.Build.cs
+++ b/Engine/Plugins/BlueprintEditorLibrary/Source/BlueprintEditorLibrary/BlueprintEditorLibrary.Build.cs
@@ -35,6 +35,7 @@ public class BlueprintEditorLibrary : ModuleRules
 				"SlateCore",
 				"UnrealEd",
 				"Kismet",
+				"AnimGraph"
 			}
 		);
 	}
diff --git a/Engine/Plugins/BlueprintEditorLibrary/Source/BlueprintEditorLibrary/Private/BlueprintEditorLibrary.cpp b/Engine/Plugins/BlueprintEditorLibrary/Source/BlueprintEditorLibrary/Private/BlueprintEditorLibrary.cpp
--- a/Engine/Plugins/BlueprintEditorLibrary/Source/BlueprintEditorLibrary/Private/BlueprintEditorLibrary.cpp
+++ b/Engine/Plugins/BlueprintEditorLibrary/Source/BlueprintEditorLibrary/Private/BlueprintEditorLibrary.cpp
@@ -14,6 +14,7 @@
 #include "Kismet2/KismetEditorUtilities.h"
 #include "BlueprintNodeSpawner.h"
 #include "EdGraphUtilities.h"
+#include "AnimGraphNode_Base.h"
 
 DEFINE_LOG_CATEGORY(LogBlueprintEditorLib);
 
@@ -61,6 +62,30 @@ namespace InternalBlueprintEditorLibrary
 			NewNode->ReconstructNode();
 		}
 	}
+
+	/**
+	* Returns true if any of these nodes pins have any links. Does not check for 
+	* a default value on pins
+	*
+	* @param Node		The node to check
+	*
+	* @return bool		True if the node has any links, false otherwise.
+	*/
+	static bool NodeHasAnyConnections(const UEdGraphNode* Node)
+	{
+		if (Node)
+		{
+			for (const UEdGraphPin* Pin : Node->Pins)
+			{
+				if (Pin && Pin->LinkedTo.Num() > 0)
+				{
+					return true;
+				}
+			}
+		}
+
+		return false;
+	}
 };
 
 ///////////////////////////////////////////////////////////
@@ -304,6 +329,55 @@ void UBlueprintEditorLibrary::RemoveFunctionGraph(UBlueprint* Blueprint, FName F
 	}
 }
 
+void UBlueprintEditorLibrary::RemoveUnusedNodes(UBlueprint* Blueprint)
+{
+	if (!Blueprint)
+	{
+		return;
+	}
+
+	TArray<UEdGraph*> AllGraphs;
+	Blueprint->GetAllGraphs(AllGraphs);
+	Blueprint->Modify();
+
+	for (UEdGraph* Graph : AllGraphs)
+	{
+		// Skip non-editable graphs
+		if (!Graph || FBlueprintEditorUtils::IsGraphReadOnly(Graph))
+		{
+			continue;
+		}
+
+		Graph->Modify();
+		int32 NumNodesRemoved = 0;
+
+		for (int32 i = Graph->Nodes.Num() - 1; i >= 0; --i)
+		{
+			UEdGraphNode* Node = Graph->Nodes[i];
+
+			// We only want to delete user facing nodes because this is meant 
+			// to be a BP refactoring/cleanup tool. Anim graph nodes can still 
+			// be valid with no pin connections made to them
+			if (Node->CanUserDeleteNode() && 
+				!Node->IsA<UAnimGraphNode_Base>() && 
+				!InternalBlueprintEditorLibrary::NodeHasAnyConnections(Node))
+			{
+				Node->BreakAllNodeLinks();
+				Graph->RemoveNode(Node);
+				++NumNodesRemoved;
+			}
+		}
+
+		// Notify a change to the graph if nodes have been removed
+		if (NumNodesRemoved > 0)
+		{
+			Graph->NotifyGraphChanged();
+		}
+	}
+
+	FBlueprintEditorUtils::MarkBlueprintAsStructurallyModified(Blueprint);
+}
+
 void UBlueprintEditorLibrary::RemoveGraph(UBlueprint* Blueprint, UEdGraph* Graph)
 {
 	if (!Blueprint || !Graph)
diff --git a/Engine/Plugins/BlueprintEditorLibrary/Source/BlueprintEditorLibrary/Public/BlueprintEditorLibrary.h b/Engine/Plugins/BlueprintEditorLibrary/Source/BlueprintEditorLibrary/Public/BlueprintEditorLibrary.h
--- a/Engine/Plugins/BlueprintEditorLibrary/Source/BlueprintEditorLibrary/Public/BlueprintEditorLibrary.h
+++ b/Engine/Plugins/BlueprintEditorLibrary/Source/BlueprintEditorLibrary/Public/BlueprintEditorLibrary.h
@@ -96,6 +96,14 @@ public:
 	UFUNCTION(BlueprintCallable, Category = "Blueprint Upgrade Tools")
 	static void RemoveFunctionGraph(UBlueprint* Blueprint, FName FuncName);
 
+	/**
+	* Remove any nodes in this blueprint that have no connections made to them.
+	*
+	* @param Blueprint		The blueprint to remove the nodes from
+	*/
+	UFUNCTION(BlueprintCallable, Category = "Blueprint Upgrade Tools")
+	static void RemoveUnusedNodes(UBlueprint* Blueprint);
+
 	/** 
 	* Removes the given graph from the blueprint if possible 
 	* 
