Blueprint editor lib: Add a RemoveUnusedVariables function to remove any unused variables from a blueprint

#rnx

diff --git a/Engine/Plugins/BlueprintEditorLibrary/Source/BlueprintEditorLibrary/Private/BlueprintEditorLibrary.cpp b/Engine/Plugins/BlueprintEditorLibrary/Source/BlueprintEditorLibrary/Private/BlueprintEditorLibrary.cpp
--- a/Engine/Plugins/BlueprintEditorLibrary/Source/BlueprintEditorLibrary/Private/BlueprintEditorLibrary.cpp
+++ b/Engine/Plugins/BlueprintEditorLibrary/Source/BlueprintEditorLibrary/Private/BlueprintEditorLibrary.cpp
@@ -15,6 +15,9 @@
 #include "BlueprintNodeSpawner.h"
 #include "EdGraphUtilities.h"
 #include "AnimGraphNode_Base.h"
+#include "Components/TimelineComponent.h"
+
+#define LOCTEXT_NAMESPACE "BlueprintEditorLibrary"
 
 DEFINE_LOG_CATEGORY(LogBlueprintEditorLib);
 
@@ -444,4 +447,86 @@ void UBlueprintEditorLibrary::ReparentBlueprint(UBlueprint* Blueprint, UClass* N
 	FBlueprintEditorUtils::MarkBlueprintAsModified(Blueprint);
 
 	CompileBlueprint(Blueprint);
-}
\ No newline at end of file
+}
+
+bool UBlueprintEditorLibrary::GatherUnusedVariables(const UBlueprint* Blueprint, TArray<FProperty*>& OutProperties)
+{
+	if (!Blueprint)
+	{
+		return false;
+	}
+
+	bool bHasAtLeastOneVariableToCheck = false;
+
+	for (TFieldIterator<FProperty> PropertyIt(Blueprint->SkeletonGeneratedClass, EFieldIteratorFlags::ExcludeSuper); PropertyIt; ++PropertyIt)
+	{
+		FProperty* Property = *PropertyIt;
+		// Don't show delegate properties, there is special handling for these
+		const bool bDelegateProp = Property->IsA(FDelegateProperty::StaticClass()) || Property->IsA(FMulticastDelegateProperty::StaticClass());
+		const bool bShouldShowProp = (!Property->HasAnyPropertyFlags(CPF_Parm) && Property->HasAllPropertyFlags(CPF_BlueprintVisible) && !bDelegateProp);
+
+		if (bShouldShowProp)
+		{
+			bHasAtLeastOneVariableToCheck = true;
+			FName VarName = Property->GetFName();
+
+			const int32 VarInfoIndex = FBlueprintEditorUtils::FindNewVariableIndex(Blueprint, VarName);
+			const bool bHasVarInfo = (VarInfoIndex != INDEX_NONE);
+
+			const FObjectPropertyBase* ObjectProperty = CastField<const FObjectPropertyBase>(Property);
+			bool bIsTimeline = ObjectProperty &&
+				ObjectProperty->PropertyClass &&
+				ObjectProperty->PropertyClass->IsChildOf(UTimelineComponent::StaticClass());
+			if (!bIsTimeline && bHasVarInfo && !FBlueprintEditorUtils::IsVariableUsed(Blueprint, VarName))
+			{
+				OutProperties.Add(Property);
+			}
+		}
+	}
+
+	return bHasAtLeastOneVariableToCheck;
+}
+
+int32 UBlueprintEditorLibrary::RemoveUnusedVariables(UBlueprint* Blueprint)
+{
+	if (!Blueprint)
+	{
+		return 0;
+	}
+
+	// Gather FProperties from this BP and see if we can remove any
+	TArray<FProperty*> VariableProperties;
+	UBlueprintEditorLibrary::GatherUnusedVariables(Blueprint, VariableProperties);
+	
+	// No variables can be removed from this blueprint
+	if (VariableProperties.Num() == 0)
+	{
+		return 0;
+	}
+
+	// Get the variables by name so that we can bulk remove them and print them out to the log
+	TArray<FName> VariableNames;
+	FString PropertyList;
+	VariableNames.Reserve(VariableProperties.Num());
+	for (int32 Index = 0; Index < VariableProperties.Num(); ++Index)
+	{
+		VariableNames.Add(VariableProperties[Index]->GetFName());
+		if (PropertyList.IsEmpty())
+		{
+			PropertyList = UEditorEngine::GetFriendlyName(VariableProperties[Index]);
+		}
+		else
+		{
+			PropertyList += FString::Printf(TEXT(", %s"), *UEditorEngine::GetFriendlyName(VariableProperties[Index]));
+		}
+	}
+
+	const int32 NumRemovedVars = VariableNames.Num();
+	// Remove the variables by name
+	FBlueprintEditorUtils::BulkRemoveMemberVariables(Blueprint, VariableNames);
+
+	UE_LOG(LogBlueprintEditorLib, Log, TEXT("The following variable(s) were deleted successfully: %s."), *PropertyList);
+	return NumRemovedVars;
+}
+
+#undef LOCTEXT_NAMESPACE	// "BlueprintEditorLibrary"
\ No newline at end of file
diff --git a/Engine/Plugins/BlueprintEditorLibrary/Source/BlueprintEditorLibrary/Public/BlueprintEditorLibrary.h b/Engine/Plugins/BlueprintEditorLibrary/Source/BlueprintEditorLibrary/Public/BlueprintEditorLibrary.h
--- a/Engine/Plugins/BlueprintEditorLibrary/Source/BlueprintEditorLibrary/Public/BlueprintEditorLibrary.h
+++ b/Engine/Plugins/BlueprintEditorLibrary/Source/BlueprintEditorLibrary/Public/BlueprintEditorLibrary.h
@@ -149,4 +149,26 @@ public:
 	*/
 	UFUNCTION(BlueprintCallable, Category = "Blueprint Upgrade Tools")
 	static void ReparentBlueprint(UBlueprint* Blueprint, UClass* NewParentClass);
+
+	/**
+	* Gathers any unused blueprint variables and populates the given array of FPropertys
+	*
+	* @param Blueprint			The blueprint to check
+	* @param OutProperties		Out array of unused FProperty*'s
+	*
+	* @return					True if variables were checked on this blueprint, false otherwise.
+	*/
+	static bool GatherUnusedVariables(const UBlueprint* Blueprint, TArray<FProperty*>& OutProperties);
+
+	/**
+	* Deletes any unused blueprint created variables the given blueprint.
+	* An Unused variable is any BP variable that is not referenced in any 
+	* blueprint graphs
+	* 
+	* @param Blueprint			Blueprint that you would like to remove variables from
+	*
+	* @return					Number of variables removed
+	*/
+	UFUNCTION(BlueprintCallable, Category = "Blueprint Upgrade Tools")
+	static int32 RemoveUnusedVariables(UBlueprint* Blueprint);
 };
\ No newline at end of file
