remove unused dependies in bp eidtor library build.cs

diff --git a/Engine/Plugins/BlueprintEditorLibrary/Source/BlueprintEditorLibrary/BlueprintEditorLibrary.Build.cs b/Engine/Plugins/BlueprintEditorLibrary/Source/BlueprintEditorLibrary/BlueprintEditorLibrary.Build.cs
--- a/Engine/Plugins/BlueprintEditorLibrary/Source/BlueprintEditorLibrary/BlueprintEditorLibrary.Build.cs
+++ b/Engine/Plugins/BlueprintEditorLibrary/Source/BlueprintEditorLibrary/BlueprintEditorLibrary.Build.cs
@@ -31,10 +31,7 @@ public class BlueprintEditorLibrary : ModuleRules
 			{
 				"CoreUObject",
 				"Engine",
-				"Slate",
-				"SlateCore",
 				"UnrealEd",
-				"Kismet",
 				"AnimGraph"
 			}
 		);
