remove unused AddComponent function from BP Editor lib

#rnx

diff --git a/Engine/Plugins/BlueprintEditorLibrary/Source/BlueprintEditorLibrary/Private/BlueprintEditorLibrary.cpp b/Engine/Plugins/BlueprintEditorLibrary/Source/BlueprintEditorLibrary/Private/BlueprintEditorLibrary.cpp
--- a/Engine/Plugins/BlueprintEditorLibrary/Source/BlueprintEditorLibrary/Private/BlueprintEditorLibrary.cpp
+++ b/Engine/Plugins/BlueprintEditorLibrary/Source/BlueprintEditorLibrary/Private/BlueprintEditorLibrary.cpp
@@ -401,20 +401,6 @@ void UBlueprintEditorLibrary::RenameGraph(UEdGraph* Graph, const FString& NewNam
 	FBlueprintEditorUtils::RenameGraph(Graph, NewNameStr);
 }
 
-void UBlueprintEditorLibrary::AddComponent(UBlueprint* Blueprint, TSubclassOf<UActorComponent> ComponentClass)
-{
-	if (!Blueprint || !ComponentClass)
-	{
-		return;
-	}
-	//UClass* ActorClass = Blueprint->GeneratedClass;
-	//TSubclassOf<UActorComponent> MatchingComponentClassForAsset = FComponentAssetBrokerage::GetPrimaryComponentForAsset(AssetClass);
-
-	//AActor* CDO = Blueprint->GeneratedClass->
-	// #TODO_BH Will need to discuss how to handle this correctly, there is a lot of
-	// logic in SCS editor that I have questions about
-}
-
 UBlueprint* UBlueprintEditorLibrary::GetBlueprintAsset(UObject* Object)
 {
 	return Cast<UBlueprint>(Object);
diff --git a/Engine/Plugins/BlueprintEditorLibrary/Source/BlueprintEditorLibrary/Public/BlueprintEditorLibrary.h b/Engine/Plugins/BlueprintEditorLibrary/Source/BlueprintEditorLibrary/Public/BlueprintEditorLibrary.h
--- a/Engine/Plugins/BlueprintEditorLibrary/Source/BlueprintEditorLibrary/Public/BlueprintEditorLibrary.h
+++ b/Engine/Plugins/BlueprintEditorLibrary/Source/BlueprintEditorLibrary/Public/BlueprintEditorLibrary.h
@@ -122,15 +122,6 @@ public:
 	UFUNCTION(BlueprintCallable, Category = "Blueprint Upgrade Tools")
 	static void RenameGraph(UEdGraph* Graph, const FString& NewNameStr);
 
-	/**
-	* Add an Actor Component to the given blueprint
-	*
-	* @param Blueprint			Blueprint to add the component to
-	* @param ComponentClass		Type of component to add
-	*/
-	UFUNCTION(BlueprintCallable, Category = "Blueprint Upgrade Tools")
-	static void AddComponent(UBlueprint* Blueprint, TSubclassOf<UActorComponent> ComponentClass);
-
 	/**
 	* Gets the UBlueprint version of the given object if possible.
 	*
