Search all BP graphs when calling FindGraph instead of just ubergraphs for the BP editor library

#rnx

diff --git a/Engine/Plugins/BlueprintEditorLibrary/Source/BlueprintEditorLibrary/Private/BlueprintEditorLibrary.cpp b/Engine/Plugins/BlueprintEditorLibrary/Source/BlueprintEditorLibrary/Private/BlueprintEditorLibrary.cpp
--- a/Engine/Plugins/BlueprintEditorLibrary/Source/BlueprintEditorLibrary/Private/BlueprintEditorLibrary.cpp
+++ b/Engine/Plugins/BlueprintEditorLibrary/Source/BlueprintEditorLibrary/Private/BlueprintEditorLibrary.cpp
@@ -152,9 +152,12 @@ UEdGraph* UBlueprintEditorLibrary::FindEventGraph(UBlueprint* Blueprint)
 
 UEdGraph* UBlueprintEditorLibrary::FindGraph(UBlueprint* Blueprint, FName GraphName)
 {
-	if (Blueprint)
+	if (Blueprint && !GraphName.IsNone())
 	{
-		for (UEdGraph* CurrentGraph : Blueprint->UbergraphPages)
+		TArray<UEdGraph*> AllGraphs;
+		Blueprint->GetAllGraphs(AllGraphs);
+
+		for (UEdGraph* CurrentGraph : AllGraphs)
 		{
 			if (CurrentGraph->GetFName() == GraphName)
 			{
