Rename RenameVariableReferences to ReplaceVariableReferences in blueprint editor for better description on the node.
Add checks for if the new variable names are valid or not.

#rnx

diff --git a/Engine/Plugins/BlueprintEditorLibrary/Source/BlueprintEditorLibrary/Private/BlueprintEditorLibrary.cpp b/Engine/Plugins/BlueprintEditorLibrary/Source/BlueprintEditorLibrary/Private/BlueprintEditorLibrary.cpp
--- a/Engine/Plugins/BlueprintEditorLibrary/Source/BlueprintEditorLibrary/Private/BlueprintEditorLibrary.cpp
+++ b/Engine/Plugins/BlueprintEditorLibrary/Source/BlueprintEditorLibrary/Private/BlueprintEditorLibrary.cpp
@@ -99,9 +99,9 @@ UBlueprintEditorLibrary::UBlueprintEditorLibrary(const FObjectInitializer& Objec
 {
 }
 
-void UBlueprintEditorLibrary::RenameVariableReferences(UBlueprint* Blueprint, const FName OldVarName, const FName NewVarName)
+void UBlueprintEditorLibrary::ReplaceVariableReferences(UBlueprint* Blueprint, const FName OldVarName, const FName NewVarName)
 {
-	if (!Blueprint)
+	if (!Blueprint || OldVarName.IsNone() || NewVarName.IsNone())
 	{
 		return;
 	}
diff --git a/Engine/Plugins/BlueprintEditorLibrary/Source/BlueprintEditorLibrary/Public/BlueprintEditorLibrary.h b/Engine/Plugins/BlueprintEditorLibrary/Source/BlueprintEditorLibrary/Public/BlueprintEditorLibrary.h
--- a/Engine/Plugins/BlueprintEditorLibrary/Source/BlueprintEditorLibrary/Public/BlueprintEditorLibrary.h
+++ b/Engine/Plugins/BlueprintEditorLibrary/Source/BlueprintEditorLibrary/Public/BlueprintEditorLibrary.h
@@ -25,7 +25,7 @@ public:
 	* @param NewVarName		The new variable that will be used in the old one's place
 	*/
 	UFUNCTION(BlueprintCallable, Category = "Blueprint Upgrade Tools")
-	static void RenameVariableReferences(UBlueprint* Blueprint, const FName OldVarName, const FName NewVarName);
+	static void ReplaceVariableReferences(UBlueprint* Blueprint, const FName OldVarName, const FName NewVarName);
 
 	/**
 	* Replace any existing k2node class with the new one, attempting to remap pins as we go
