Update the function description for FindEventGraph in the blueprint editor library to specifiy that it will only return one primiary event graph, not an array.

#rnx

diff --git a/Engine/Plugins/BlueprintEditorLibrary/Source/BlueprintEditorLibrary/Public/BlueprintEditorLibrary.h b/Engine/Plugins/BlueprintEditorLibrary/Source/BlueprintEditorLibrary/Public/BlueprintEditorLibrary.h
--- a/Engine/Plugins/BlueprintEditorLibrary/Source/BlueprintEditorLibrary/Public/BlueprintEditorLibrary.h
+++ b/Engine/Plugins/BlueprintEditorLibrary/Source/BlueprintEditorLibrary/Public/BlueprintEditorLibrary.h
@@ -38,7 +38,8 @@ public:
 	static void ReplaceK2Nodes(UBlueprint* Blueprint, TSubclassOf<UEdGraphNode> OldNodeType, TSubclassOf<UEdGraphNode> NewNodeType);
 
 	/**
-	* Finds the event graph of the given blueprint. Null if it doesn't have one.
+	* Finds the event graph of the given blueprint. Null if it doesn't have one. This will only return
+	* the primary event graph of the blueprint (the graph named "EventGraph").
 	*
 	* @param Blueprint		Blueprint to search for the event graph on
 	*
